Nounity fixes

diff --git a/Engine/Plugins/Editor/ContentBrowser/ContentBrowserClassDataSource/Source/ContentBrowserClassDataSource/Private/NativeClassHierarchy.h b/Engine/Plugins/Editor/ContentBrowser/ContentBrowserClassDataSource/Source/ContentBrowserClassDataSource/Private/NativeClassHierarchy.h
--- a/Engine/Plugins/Editor/ContentBrowser/ContentBrowserClassDataSource/Source/ContentBrowserClassDataSource/Private/NativeClassHierarchy.h
+++ b/Engine/Plugins/Editor/ContentBrowser/ContentBrowserClassDataSource/Source/ContentBrowserClassDataSource/Private/NativeClassHierarchy.h
@@ -6,6 +6,7 @@
 #include "Containers/ArrayView.h"
 #include "HAL/PlatformTime.h"
 #include "Interfaces/IPluginManager.h"
+#include "UObject/UObjectGlobals.h"
 
 /**
  * Type of hierarchy node
