Editor: Reduced per-module editor startup log spam by changing LogNativeClassHierarchy to Log from Verbose

diff --git a/Engine/Plugins/Editor/ContentBrowser/ContentBrowserClassDataSource/Source/ContentBrowserClassDataSource/Private/NativeClassHierarchy.cpp b/Engine/Plugins/Editor/ContentBrowser/ContentBrowserClassDataSource/Source/ContentBrowserClassDataSource/Private/NativeClassHierarchy.cpp
--- a/Engine/Plugins/Editor/ContentBrowser/ContentBrowserClassDataSource/Source/ContentBrowserClassDataSource/Private/NativeClassHierarchy.cpp
+++ b/Engine/Plugins/Editor/ContentBrowser/ContentBrowserClassDataSource/Source/ContentBrowserClassDataSource/Private/NativeClassHierarchy.cpp
@@ -15,7 +15,7 @@
 #include "Interfaces/IProjectManager.h"
 #include "ProjectDescriptor.h"
 
-DEFINE_LOG_CATEGORY_STATIC(LogNativeClassHierarchy, Verbose, All);
+DEFINE_LOG_CATEGORY_STATIC(LogNativeClassHierarchy, Log, All);
 
 TSharedRef<FNativeClassHierarchyNode> FNativeClassHierarchyNode::MakeFolderEntry(FName InEntryName, FString InEntryPath)
 {
