#UE Remove GetPluginModules from FNativeClassHierarchy to improve editor startup performance. Modern module manager code caches the map of modulenames to plugins O(1) so this is no longer needed. See GetModuleOwnerPlugin

diff --git a/Engine/Plugins/Editor/ContentBrowser/ContentBrowserClassDataSource/Source/ContentBrowserClassDataSource/Private/NativeClassHierarchy.cpp b/Engine/Plugins/Editor/ContentBrowser/ContentBrowserClassDataSource/Source/ContentBrowserClassDataSource/Private/NativeClassHierarchy.cpp
--- a/Engine/Plugins/Editor/ContentBrowser/ContentBrowserClassDataSource/Source/ContentBrowserClassDataSource/Private/NativeClassHierarchy.cpp
+++ b/Engine/Plugins/Editor/ContentBrowser/ContentBrowserClassDataSource/Source/ContentBrowserClassDataSource/Private/NativeClassHierarchy.cpp
@@ -319,12 +319,11 @@ void FNativeClassHierarchy::PopulateHierarchy()
 	RootNodes.Empty();
 
 	TSet<FName> GameModules = GetGameModules();
-	TMap<FName, FNativeClassHierarchyPluginModuleInfo> PluginModules = GetPluginModules();
-
+	
 	for(TObjectIterator<UClass> ClassIt; ClassIt; ++ClassIt)
 	{
 		UClass* const CurrentClass = *ClassIt;
-		AddClass(CurrentClass, GameModules, PluginModules, AddClassMetrics);
+		AddClass(CurrentClass, GameModules, AddClassMetrics);
 	}
 
 	UE_LOG(LogNativeClassHierarchy, Verbose, TEXT("Native class hierarchy populated in %0.4f seconds. Added %d classes and %d folders."), FPlatformTime::Seconds() - AddClassMetrics.StartTime, AddClassMetrics.NumClassesAdded, AddClassMetrics.NumFoldersAdded);
@@ -344,8 +343,7 @@ void FNativeClassHierarchy::AddClassesForModule(const FName& InModuleName)
 	}
 
 	TSet<FName> GameModules = GetGameModules();
-	TMap<FName, FNativeClassHierarchyPluginModuleInfo> PluginModules = GetPluginModules();
-
+	
 	TArray<UObject*> PackageObjects;
 	GetObjectsWithOuter(ClassPackage, PackageObjects, false);
 	for(UObject* Object : PackageObjects)
@@ -353,7 +351,7 @@ void FNativeClassHierarchy::AddClassesForModule(const FName& InModuleName)
 		UClass* const CurrentClass = Cast<UClass>(Object);
 		if(CurrentClass)
 		{
-			AddClass(CurrentClass, GameModules, PluginModules, AddClassMetrics);
+			AddClass(CurrentClass, GameModules, AddClassMetrics);
 		}
 	}
 
@@ -387,7 +385,7 @@ void FNativeClassHierarchy::RemoveClassesForModule(const FName& InModuleName)
 	}
 }
 
-void FNativeClassHierarchy::AddClass(UClass* InClass, const TSet<FName>& InGameModules, const TMap<FName, FNativeClassHierarchyPluginModuleInfo>& InPluginModules, FAddClassMetrics& AddClassMetrics)
+void FNativeClassHierarchy::AddClass(UClass* InClass, const TSet<FName>& InGameModules, FAddClassMetrics& AddClassMetrics)
 {
 	// Ignore deprecated and temporary classes
 	if(InClass->HasAnyClassFlags(CLASS_Deprecated | CLASS_NewerVersionExists) || FKismetEditorUtilities::IsClassABlueprintSkeleton(InClass))
@@ -410,7 +408,7 @@ void FNativeClassHierarchy::AddClass(UClass* InClass, const TSet<FName>& InGameM
 
 	// Work out which root this class should go under
 	EPluginLoadedFrom WhereLoadedFrom;
-	const FName RootNodeName = GetClassPathRootForModule(ClassModuleName, InGameModules, InPluginModules, WhereLoadedFrom);
+	const FName RootNodeName = GetClassPathRootForModule(ClassModuleName, InGameModules, WhereLoadedFrom);
 
 	// Work out the final path to this class within the hierarchy (which isn't the same as the path on disk)
 	const FString ClassModuleRelativePath = ClassModuleRelativeIncludePath.Left(ClassModuleRelativeIncludePath.Find(TEXT("/"), ESearchCase::CaseSensitive, ESearchDir::FromEnd));
@@ -526,15 +524,14 @@ bool FNativeClassHierarchy::GetClassPath(const UClass* InClass, FString& OutClas
 		return false;
 	}
 
-	if (InCache.GameModules.Num() == 0 && InCache.PluginModules.Num() == 0)
+	if (InCache.GameModules.Num() == 0)
 	{
 		InCache.GameModules = GetGameModules();
-		InCache.PluginModules = GetPluginModules();
 	}
 
 	// Work out which root this class should go under
 	EPluginLoadedFrom WhereLoadedFrom;
-	const FName RootNodeName = GetClassPathRootForModule(ClassModuleName, InCache.GameModules, InCache.PluginModules, WhereLoadedFrom);
+	const FName RootNodeName = GetClassPathRootForModule(ClassModuleName, InCache.GameModules, WhereLoadedFrom);
 
 	// Work out the final path to this class within the hierarchy (which isn't the same as the path on disk)
 	const FString ClassModuleRelativePath = ClassModuleRelativeIncludePath.Left(ClassModuleRelativeIncludePath.Find(TEXT("/"), ESearchCase::CaseSensitive, ESearchDir::FromEnd));
@@ -587,7 +584,7 @@ FName FNativeClassHierarchy::GetClassModuleName(const UClass* InClass)
 	return NAME_None;
 }
 
-FName FNativeClassHierarchy::GetClassPathRootForModule(const FName& InModuleName, const TSet<FName>& InGameModules, const TMap<FName, FNativeClassHierarchyPluginModuleInfo>& InPluginModules, EPluginLoadedFrom& OutWhereLoadedFrom)
+FName FNativeClassHierarchy::GetClassPathRootForModule(const FName& InModuleName, const TSet<FName>& InGameModules, EPluginLoadedFrom& OutWhereLoadedFrom)
 {
 	static const FName EngineRootNodeName = "Classes_Engine";
 	static const FName GameRootNodeName = "Classes_Game";
@@ -601,11 +598,14 @@ FName FNativeClassHierarchy::GetClassPathRootForModule(const FName& InModuleName
 		RootNodeName = GameRootNodeName;
 		OutWhereLoadedFrom = EPluginLoadedFrom::Project;
 	}
-	else if(InPluginModules.Contains(InModuleName))
+	else
 	{
-		const FNativeClassHierarchyPluginModuleInfo PluginInfo = InPluginModules.FindRef(InModuleName);
-		RootNodeName = FName(*(FString(TEXT("Classes_")) + PluginInfo.Name.ToString()));
-		OutWhereLoadedFrom = PluginInfo.LoadedFrom;
+		TSharedPtr<IPlugin> ModulePlugin = IPluginManager::Get().GetModuleOwnerPlugin(InModuleName);
+		if (ModulePlugin.IsValid())
+		{
+			RootNodeName = FName(*(TEXT("Classes_") + ModulePlugin->GetName()));
+			OutWhereLoadedFrom = ModulePlugin->GetLoadedFrom();
+		}
 	}
 
 	return RootNodeName;
@@ -627,27 +627,4 @@ TSet<FName> FNativeClassHierarchy::GetGameModules()
 	}
 
 	return GameModules;
-}
-
-TMap<FName, FNativeClassHierarchyPluginModuleInfo> FNativeClassHierarchy::GetPluginModules()
-{
-	IPluginManager& PluginManager = IPluginManager::Get();
-
-	// Build up a map of plugin modules -> plugin names - used to work out which modules populate a given Classes_PluginName
-	TMap<FName, FNativeClassHierarchyPluginModuleInfo> PluginModules;
-	{
-		TArray<TSharedRef<IPlugin>> Plugins = PluginManager.GetDiscoveredPlugins();
-		for(const TSharedRef<IPlugin>& Plugin: Plugins)
-		{
-			for(const FModuleDescriptor& PluginModule: Plugin->GetDescriptor().Modules)
-			{
-				FNativeClassHierarchyPluginModuleInfo Info;
-				Info.Name = FName(*Plugin->GetName());
-				Info.LoadedFrom = Plugin->GetLoadedFrom();
-				PluginModules.Add(PluginModule.Name, Info);
-			}
-		}
-	}
-
-	return PluginModules;
-}
+}
\ No newline at end of file
diff --git a/Engine/Plugins/Editor/ContentBrowser/ContentBrowserClassDataSource/Source/ContentBrowserClassDataSource/Private/NativeClassHierarchy.h b/Engine/Plugins/Editor/ContentBrowser/ContentBrowserClassDataSource/Source/ContentBrowserClassDataSource/Private/NativeClassHierarchy.h
--- a/Engine/Plugins/Editor/ContentBrowser/ContentBrowserClassDataSource/Source/ContentBrowserClassDataSource/Private/NativeClassHierarchy.h
+++ b/Engine/Plugins/Editor/ContentBrowser/ContentBrowserClassDataSource/Source/ContentBrowserClassDataSource/Private/NativeClassHierarchy.h
@@ -17,30 +17,16 @@ enum class ENativeClassHierarchyNodeType : uint8
 	Class,
 };
 
-/**
- *  Contains high-level information about the plugin module relevant to the native class hierarchy.                                                                      
- */
-struct FNativeClassHierarchyPluginModuleInfo
-{
-	/** Name of the module*/
-	FName Name;
-
-	/** Indicator of where the module was loaded from (Engine or GameProject)*/
-	EPluginLoadedFrom LoadedFrom;
-};
-
 /**
  * Cache to avoid regenerating some necessary data during repeated calls to GetClassPathRootForModule during enumerate
  */
 struct FNativeClassHierarchyGetClassPathCache
 {
 	TSet<FName> GameModules;
-	TMap<FName, FNativeClassHierarchyPluginModuleInfo> PluginModules;
 
 	void Reset()
 	{
 		GameModules.Reset();
-		PluginModules.Reset();
 	}
 };
 
@@ -361,10 +347,9 @@ private:
 	 *
 	 * @param InClass - The class that is to be added
 	 * @param InGameModules - The list of modules that belong to the "/Classes_Game" root
-	 * @param InPluginModules - The list of modules that belong to neither the "/Classes_Game" or "/Classes_Engine" roots
 	 * @param AddClassMetrics - Metrics to update as new classes and folders are added (used to report population performance)
 	 */
-	void AddClass(UClass* InClass, const TSet<FName>& InGameModules, const TMap<FName, FNativeClassHierarchyPluginModuleInfo>& InPluginModules, FAddClassMetrics& AddClassMetrics);
+	void AddClass(UClass* InClass, const TSet<FName>& InGameModules, FAddClassMetrics& AddClassMetrics);
 
 	/**
 	 * Called when we're notified that a module has changed status
@@ -395,23 +380,17 @@ private:
 	 *
 	 * @param InModuleName - The module we want to find the root path for
 	 * @param InGameModules - The list of modules that belong to the "/Classes_Game" root
-	 * @param InPluginModules - The list of modules that belong to neither the "/Classes_Game" or "/Classes_Engine" roots
 	 * @param WhereLoadedFrom - Determine where this class was loaded from.
 	 * 
 	 * @return The name of the root path to use, eg "/Classes_Game"
 	 */
-	static FName GetClassPathRootForModule(const FName& InModuleName, const TSet<FName>& InGameModules, const TMap<FName, FNativeClassHierarchyPluginModuleInfo>& InPluginModules, EPluginLoadedFrom& OutWhereLoadedFrom);
+	static FName GetClassPathRootForModule(const FName& InModuleName, const TSet<FName>& InGameModules, EPluginLoadedFrom& OutWhereLoadedFrom);
 
 	/**
 	 * Calculate the list of modules that belong to the "/Classes_Game" root
 	 */
 	static TSet<FName> GetGameModules();
 
-	/**
-	 * Calculate the list of modules that belong to neither the "/Classes_Game" or "/Classes_Engine" roots
-	 */
-	static TMap<FName, FNativeClassHierarchyPluginModuleInfo> GetPluginModules();
-
 	/** Root level nodes corresponding to the root folders used by the Content Browser, eg) Classes_Engine, Classes_Game, etc */
 	TMap<FName, TSharedPtr<FNativeClassHierarchyNode>> RootNodes;
 
