Apply flag EContentBrowserItemFlags::Category_Plugin to classes from plugins

diff --git a/Engine/Plugins/Editor/ContentBrowser/ContentBrowserClassDataSource/Source/ContentBrowserClassDataSource/Private/ContentBrowserClassDataCore.cpp b/Engine/Plugins/Editor/ContentBrowser/ContentBrowserClassDataSource/Source/ContentBrowserClassDataSource/Private/ContentBrowserClassDataCore.cpp
--- a/Engine/Plugins/Editor/ContentBrowser/ContentBrowserClassDataSource/Source/ContentBrowserClassDataSource/Private/ContentBrowserClassDataCore.cpp
+++ b/Engine/Plugins/Editor/ContentBrowser/ContentBrowserClassDataSource/Source/ContentBrowserClassDataSource/Private/ContentBrowserClassDataCore.cpp
@@ -100,7 +100,7 @@ bool IsPluginClass(const FName InPath)
 	return AssetViewUtils::IsPluginFolder(PathStr);
 }
 
-FContentBrowserItemData CreateClassFolderItem(UContentBrowserDataSource* InOwnerDataSource, const FName InVirtualPath, const FName InFolderPath)
+FContentBrowserItemData CreateClassFolderItem(UContentBrowserDataSource* InOwnerDataSource, const FName InVirtualPath, const FName InFolderPath, const bool bIsFromPlugin)
 {
 	static const FName GameRootPath = "/Classes_Game";
 	static const FName EngineRootPath = "/Classes_Engine";
@@ -121,12 +121,22 @@ FContentBrowserItemData CreateClassFolderItem(UContentBrowserDataSource* InOwner
 		FolderDisplayNameOverride = ContentBrowserDataUtils::GetFolderItemDisplayNameOverride(InFolderPath, FolderItemName, /*bIsClassesFolder*/ true);
 	}
 
-	return FContentBrowserItemData(InOwnerDataSource, EContentBrowserItemFlags::Type_Folder | EContentBrowserItemFlags::Category_Class, InVirtualPath, *FolderItemName, MoveTemp(FolderDisplayNameOverride), MakeShared<FContentBrowserClassFolderItemDataPayload>(InFolderPath));
+	return FContentBrowserItemData(InOwnerDataSource,
+		EContentBrowserItemFlags::Type_Folder | EContentBrowserItemFlags::Category_Class | (bIsFromPlugin ? EContentBrowserItemFlags::Category_Plugin : EContentBrowserItemFlags::None),
+		InVirtualPath,
+		*FolderItemName,
+		MoveTemp(FolderDisplayNameOverride),
+		MakeShared<FContentBrowserClassFolderItemDataPayload>(InFolderPath));
 }
 
-FContentBrowserItemData CreateClassFileItem(UContentBrowserDataSource* InOwnerDataSource, const FName InVirtualPath, const FName InClassPath, UClass* InClass)
+FContentBrowserItemData CreateClassFileItem(UContentBrowserDataSource* InOwnerDataSource, const FName InVirtualPath, const FName InClassPath, UClass* InClass, const bool bIsFromPlugin)
 {
-	return FContentBrowserItemData(InOwnerDataSource, EContentBrowserItemFlags::Type_File | EContentBrowserItemFlags::Category_Class, InVirtualPath, InClass->GetFName(), FText(), MakeShared<FContentBrowserClassFileItemDataPayload>(InClassPath, InClass));
+	return FContentBrowserItemData(InOwnerDataSource,
+		EContentBrowserItemFlags::Type_File | EContentBrowserItemFlags::Category_Class | (bIsFromPlugin ? EContentBrowserItemFlags::Category_Plugin : EContentBrowserItemFlags::None),
+		InVirtualPath,
+		InClass->GetFName(),
+		FText(),
+		MakeShared<FContentBrowserClassFileItemDataPayload>(InClassPath, InClass));
 }
 
 TSharedPtr<const FContentBrowserClassFolderItemDataPayload> GetClassFolderItemPayload(const UContentBrowserDataSource* InOwnerDataSource, const FContentBrowserItemData& InItem)
diff --git a/Engine/Plugins/Editor/ContentBrowser/ContentBrowserClassDataSource/Source/ContentBrowserClassDataSource/Private/ContentBrowserClassDataSource.cpp b/Engine/Plugins/Editor/ContentBrowser/ContentBrowserClassDataSource/Source/ContentBrowserClassDataSource/Private/ContentBrowserClassDataSource.cpp
--- a/Engine/Plugins/Editor/ContentBrowser/ContentBrowserClassDataSource/Source/ContentBrowserClassDataSource/Private/ContentBrowserClassDataSource.cpp
+++ b/Engine/Plugins/Editor/ContentBrowser/ContentBrowserClassDataSource/Source/ContentBrowserClassDataSource/Private/ContentBrowserClassDataSource.cpp
@@ -402,8 +402,10 @@ bool UContentBrowserClassDataSource::EnumerateItemsForObjects(const TArrayView<U
 			InternalPath.Reset();
 			if (NativeClassHierarchy->GetClassPath(InClass, InternalPath, NativeClassHierarchyGetClassPathCache))
 			{
+				FName ClassPath = FName(InternalPath);
+				TSharedPtr<const FNativeClassHierarchyNode> ClassNode = NativeClassHierarchy->FindNode(ClassPath, ENativeClassHierarchyNodeType::Class);
 				const FContentBrowserItemPath ContentBrowserItemPath(InternalPath, EContentBrowserPathType::Internal);
-				if (!InCallback(ContentBrowserClassData::CreateClassFileItem(this, ContentBrowserItemPath.GetVirtualPathName(), FName(InternalPath), InClass)))
+				if (!InCallback(ContentBrowserClassData::CreateClassFileItem(this, ContentBrowserItemPath.GetVirtualPathName(), ClassPath, InClass, ClassNode.IsValid() && ClassNode->LoadedFrom.IsSet())))
 				{
 					return false;
 				}
@@ -595,8 +597,9 @@ FContentBrowserItemData UContentBrowserClassDataSource::CreateClassFolderItem(co
 {
 	FName VirtualizedPath;
 	TryConvertInternalPathToVirtual(InFolderPath, VirtualizedPath);
-
-	return ContentBrowserClassData::CreateClassFolderItem(this, VirtualizedPath, InFolderPath);
+	
+	TSharedPtr<const FNativeClassHierarchyNode> FolderNode = NativeClassHierarchy->FindNode(InFolderPath, ENativeClassHierarchyNodeType::Folder);
+	return ContentBrowserClassData::CreateClassFolderItem(this, VirtualizedPath, InFolderPath, FolderNode.IsValid() && FolderNode->LoadedFrom.IsSet());
 }
 
 FContentBrowserItemData UContentBrowserClassDataSource::CreateClassFileItem(UClass* InClass, FNativeClassHierarchyGetClassPathCache& InCache)
@@ -614,7 +617,8 @@ FContentBrowserItemData UContentBrowserClassDataSource::CreateClassFileItem(UCla
 	FName VirtualizedPath;
 	TryConvertInternalPathToVirtual(ClassPath, VirtualizedPath);
 
-	return ContentBrowserClassData::CreateClassFileItem(this, VirtualizedPath, ClassPath, InClass);
+	TSharedPtr<const FNativeClassHierarchyNode> ClassNode = NativeClassHierarchy->FindNode(ClassPath, ENativeClassHierarchyNodeType::Class);
+	return ContentBrowserClassData::CreateClassFileItem(this, VirtualizedPath, ClassPath, InClass, ClassNode.IsValid() && ClassNode->LoadedFrom.IsSet());
 }
 
 TSharedPtr<const FContentBrowserClassFolderItemDataPayload> UContentBrowserClassDataSource::GetClassFolderItemPayload(const FContentBrowserItemData& InItem) const
diff --git a/Engine/Plugins/Editor/ContentBrowser/ContentBrowserClassDataSource/Source/ContentBrowserClassDataSource/Private/NativeClassHierarchy.cpp b/Engine/Plugins/Editor/ContentBrowser/ContentBrowserClassDataSource/Source/ContentBrowserClassDataSource/Private/NativeClassHierarchy.cpp
--- a/Engine/Plugins/Editor/ContentBrowser/ContentBrowserClassDataSource/Source/ContentBrowserClassDataSource/Private/NativeClassHierarchy.cpp
+++ b/Engine/Plugins/Editor/ContentBrowser/ContentBrowserClassDataSource/Source/ContentBrowserClassDataSource/Private/NativeClassHierarchy.cpp
@@ -14,17 +14,18 @@
 
 DEFINE_LOG_CATEGORY_STATIC(LogNativeClassHierarchy, Log, All);
 
-TSharedRef<FNativeClassHierarchyNode> FNativeClassHierarchyNode::MakeFolderEntry(FName InEntryName, FString InEntryPath)
+TSharedRef<FNativeClassHierarchyNode> FNativeClassHierarchyNode::MakeFolderEntry(FName InEntryName, FString InEntryPath, TOptional<EPluginLoadedFrom> LoadedFrom)
 {
 	TSharedRef<FNativeClassHierarchyNode> NewEntry = MakeShareable(new FNativeClassHierarchyNode());
 	NewEntry->Type = ENativeClassHierarchyNodeType::Folder;
 	NewEntry->Class = nullptr;
 	NewEntry->EntryName = InEntryName;
 	NewEntry->EntryPath = MoveTemp(InEntryPath);
+	NewEntry->LoadedFrom = LoadedFrom;
 	return NewEntry;
 }
 
-TSharedRef<FNativeClassHierarchyNode> FNativeClassHierarchyNode::MakeClassEntry(UClass* InClass, FName InClassModuleName, FString InClassModuleRelativePath, FString InEntryPath)
+TSharedRef<FNativeClassHierarchyNode> FNativeClassHierarchyNode::MakeClassEntry(UClass* InClass, FName InClassModuleName, FString InClassModuleRelativePath, FString InEntryPath, TOptional<EPluginLoadedFrom> LoadedFrom)
 {
 	TSharedRef<FNativeClassHierarchyNode> NewEntry = MakeShareable(new FNativeClassHierarchyNode());
 	NewEntry->Type = ENativeClassHierarchyNodeType::Class;
@@ -33,6 +34,7 @@ TSharedRef<FNativeClassHierarchyNode> FNativeClassHierarchyNode::MakeClassEntry(
 	NewEntry->ClassModuleRelativePath = MoveTemp(InClassModuleRelativePath);
 	NewEntry->EntryName = InClass->GetFName();
 	NewEntry->EntryPath = MoveTemp(InEntryPath);
+	NewEntry->LoadedFrom = LoadedFrom;
 	return NewEntry;
 }
 
@@ -404,7 +406,7 @@ void FNativeClassHierarchy::AddClass(UClass* InClass, const TSet<FName>& InGameM
 	}
 
 	// Work out which root this class should go under
-	EPluginLoadedFrom WhereLoadedFrom;
+	TOptional<EPluginLoadedFrom> WhereLoadedFrom;
 	const FName RootNodeName = GetClassPathRootForModule(ClassModuleName, InGameModules, WhereLoadedFrom);
 
 	// Work out the final path to this class within the hierarchy (which isn't the same as the path on disk)
@@ -415,7 +417,7 @@ void FNativeClassHierarchy::AddClass(UClass* InClass, const TSet<FName>& InGameM
 	TSharedPtr<FNativeClassHierarchyNode>& RootNode = RootNodes.FindOrAdd(RootNodeName);
 	if(!RootNode.IsValid())
 	{
-		RootNode = FNativeClassHierarchyNode::MakeFolderEntry(RootNodeName, TEXT("/") + RootNodeName.ToString());
+		RootNode = FNativeClassHierarchyNode::MakeFolderEntry(RootNodeName, TEXT("/") + RootNodeName.ToString(), WhereLoadedFrom);
 		RootNode->LoadedFrom = WhereLoadedFrom;
 		++AddClassMetrics.NumFoldersAdded;
 	}
@@ -430,43 +432,17 @@ void FNativeClassHierarchy::AddClass(UClass* InClass, const TSet<FName>& InGameM
 		TSharedPtr<FNativeClassHierarchyNode>& ChildNode = CurrentNode->Children.FindOrAdd(FNativeClassHierarchyNodeKey(HierarchyPathPartName, ENativeClassHierarchyNodeType::Folder));
 		if(!ChildNode.IsValid())
 		{
-			ChildNode = FNativeClassHierarchyNode::MakeFolderEntry(HierarchyPathPartName, CurrentNode->EntryPath + TEXT("/") + HierarchyPathPart);
+			ChildNode = FNativeClassHierarchyNode::MakeFolderEntry(HierarchyPathPartName, CurrentNode->EntryPath + TEXT("/") + HierarchyPathPart, WhereLoadedFrom);
 			++AddClassMetrics.NumFoldersAdded;
 		}
 		CurrentNode = ChildNode;
 	}
 
 	// Now add the final entry for the class
-	CurrentNode->AddChild(FNativeClassHierarchyNode::MakeClassEntry(InClass, ClassModuleName, ClassModuleRelativePath, CurrentNode->EntryPath + TEXT("/") + InClass->GetName()));
+	CurrentNode->AddChild(FNativeClassHierarchyNode::MakeClassEntry(InClass, ClassModuleName, ClassModuleRelativePath, CurrentNode->EntryPath + TEXT("/") + InClass->GetName(), WhereLoadedFrom));
 	++AddClassMetrics.NumClassesAdded;
 }
 
-void FNativeClassHierarchy::AddFolder(const FString& InClassPath)
-{
-	bool bHasAddedFolder = false;
-
-	// Split the class path and ensure we have nodes for each part
-	TArray<FString> ClassPathParts;
-	InClassPath.ParseIntoArray(ClassPathParts, TEXT("/"), true);
-	TSharedPtr<FNativeClassHierarchyNode> CurrentNode;
-	for(const FString& ClassPathPart : ClassPathParts)
-	{
-		const FName ClassPathPartName = *ClassPathPart;
-		TSharedPtr<FNativeClassHierarchyNode>& ChildNode = (CurrentNode.IsValid()) ? CurrentNode->Children.FindOrAdd(FNativeClassHierarchyNodeKey(ClassPathPartName, ENativeClassHierarchyNodeType::Folder)) : RootNodes.FindOrAdd(ClassPathPartName);
-		if(!ChildNode.IsValid())
-		{
-			ChildNode = FNativeClassHierarchyNode::MakeFolderEntry(ClassPathPartName, CurrentNode->EntryPath + TEXT("/") + ClassPathPart);
-			bHasAddedFolder = true;
-		}
-		CurrentNode = ChildNode;
-	}
-
-	if(bHasAddedFolder)
-	{
-		ClassHierarchyUpdatedDelegate.Broadcast();
-	}
-}
-
 bool FNativeClassHierarchy::GetFileSystemPath(const FString& InClassPath, FString& OutFileSystemPath) const
 {
 	// Split the class path into its component parts
@@ -527,7 +503,7 @@ bool FNativeClassHierarchy::GetClassPath(const UClass* InClass, FString& OutClas
 	}
 
 	// Work out which root this class should go under
-	EPluginLoadedFrom WhereLoadedFrom;
+	TOptional<EPluginLoadedFrom> WhereLoadedFrom;
 	const FName RootNodeName = GetClassPathRootForModule(ClassModuleName, InCache.GameModules, WhereLoadedFrom);
 
 	// Work out the final path to this class within the hierarchy (which isn't the same as the path on disk)
@@ -581,19 +557,17 @@ FName FNativeClassHierarchy::GetClassModuleName(const UClass* InClass)
 	return NAME_None;
 }
 
-FName FNativeClassHierarchy::GetClassPathRootForModule(const FName& InModuleName, const TSet<FName>& InGameModules, EPluginLoadedFrom& OutWhereLoadedFrom)
+FName FNativeClassHierarchy::GetClassPathRootForModule(const FName& InModuleName, const TSet<FName>& InGameModules, TOptional<EPluginLoadedFrom>& OutWhereLoadedFrom)
 {
 	static const FName EngineRootNodeName = "Classes_Engine";
 	static const FName GameRootNodeName = "Classes_Game";
 
 	// Work out which root this class should go under (anything that isn't a game or plugin module goes under engine)
 	FName RootNodeName = EngineRootNodeName;
-	OutWhereLoadedFrom = EPluginLoadedFrom::Engine;
 
 	if(InGameModules.Contains(InModuleName))
 	{
 		RootNodeName = GameRootNodeName;
-		OutWhereLoadedFrom = EPluginLoadedFrom::Project;
 	}
 	else
 	{
diff --git a/Engine/Plugins/Editor/ContentBrowser/ContentBrowserClassDataSource/Source/ContentBrowserClassDataSource/Private/NativeClassHierarchy.h b/Engine/Plugins/Editor/ContentBrowser/ContentBrowserClassDataSource/Source/ContentBrowserClassDataSource/Private/NativeClassHierarchy.h
--- a/Engine/Plugins/Editor/ContentBrowser/ContentBrowserClassDataSource/Source/ContentBrowserClassDataSource/Private/NativeClassHierarchy.h
+++ b/Engine/Plugins/Editor/ContentBrowser/ContentBrowserClassDataSource/Source/ContentBrowserClassDataSource/Private/NativeClassHierarchy.h
@@ -72,10 +72,10 @@ struct FNativeClassHierarchyNodeKey
 struct FNativeClassHierarchyNode
 {
 	/** Helper function to make a folder node entry */
-	static TSharedRef<FNativeClassHierarchyNode> MakeFolderEntry(FName InEntryName, FString InEntryPath);
+	static TSharedRef<FNativeClassHierarchyNode> MakeFolderEntry(FName InEntryName, FString InEntryPath, TOptional<EPluginLoadedFrom> LoadedFrom);
 
 	/** Helper function to make a class node entry */
-	static TSharedRef<FNativeClassHierarchyNode> MakeClassEntry(UClass* InClass, FName InClassModuleName, FString InClassModuleRelativePath, FString InEntryPath);
+	static TSharedRef<FNativeClassHierarchyNode> MakeClassEntry(UClass* InClass, FName InClassModuleName, FString InClassModuleRelativePath, FString InEntryPath, TOptional<EPluginLoadedFrom> LoadedFrom);
 
 	void AddChild(TSharedRef<FNativeClassHierarchyNode> ChildEntry);
 
@@ -101,7 +101,7 @@ struct FNativeClassHierarchyNode
 	TMap<FNativeClassHierarchyNodeKey, TSharedPtr<FNativeClassHierarchyNode>> Children;
 
 	/** Which type of plugin this data was originally loaded from (if loaded from a plugin)*/
-	EPluginLoadedFrom LoadedFrom;
+	TOptional<EPluginLoadedFrom> LoadedFrom;
 };
 
 /**
@@ -243,14 +243,6 @@ public:
 	 */
 	bool GetClassPath(const UClass* InClass, FString& OutClassPath, FNativeClassHierarchyGetClassPathCache& InCache, const bool bIncludeClassName = true) const;
 
-	/**
-	 * This will add a transient folder into the hierarchy
-	 * The folder will be lost unless a class is added to it before the hierarchy is next re-populated
-	 *
-	 * @param InClassPath - The location of the new folder (in class path form - eg) "/Classes_Game/MyGame/MyAwesomeCode")
-	 */
-	void AddFolder(const FString& InClassPath);
-
 	/**
 	 * Test if root node passes given rules
 	 */
@@ -384,7 +376,7 @@ private:
 	 * 
 	 * @return The name of the root path to use, eg "/Classes_Game"
 	 */
-	static FName GetClassPathRootForModule(const FName& InModuleName, const TSet<FName>& InGameModules, EPluginLoadedFrom& OutWhereLoadedFrom);
+	static FName GetClassPathRootForModule(const FName& InModuleName, const TSet<FName>& InGameModules, TOptional<EPluginLoadedFrom>& OutWhereLoadedFrom);
 
 	/**
 	 * Calculate the list of modules that belong to the "/Classes_Game" root
diff --git a/Engine/Plugins/Editor/ContentBrowser/ContentBrowserClassDataSource/Source/ContentBrowserClassDataSource/Public/ContentBrowserClassDataCore.h b/Engine/Plugins/Editor/ContentBrowser/ContentBrowserClassDataSource/Source/ContentBrowserClassDataSource/Public/ContentBrowserClassDataCore.h
--- a/Engine/Plugins/Editor/ContentBrowser/ContentBrowserClassDataSource/Source/ContentBrowserClassDataSource/Public/ContentBrowserClassDataCore.h
+++ b/Engine/Plugins/Editor/ContentBrowser/ContentBrowserClassDataSource/Source/ContentBrowserClassDataSource/Public/ContentBrowserClassDataCore.h
@@ -20,9 +20,9 @@ namespace ContentBrowserClassData
 
 	CONTENTBROWSERCLASSDATASOURCE_API bool IsPluginClass(const FName InPath);
 
-	CONTENTBROWSERCLASSDATASOURCE_API FContentBrowserItemData CreateClassFolderItem(UContentBrowserDataSource* InOwnerDataSource, const FName InVirtualPath, const FName InFolderPath);
+	CONTENTBROWSERCLASSDATASOURCE_API FContentBrowserItemData CreateClassFolderItem(UContentBrowserDataSource* InOwnerDataSource, const FName InVirtualPath, const FName InFolderPath, const bool bIsFromPlugin);
 
-	CONTENTBROWSERCLASSDATASOURCE_API FContentBrowserItemData CreateClassFileItem(UContentBrowserDataSource* InOwnerDataSource, const FName InVirtualPath, const FName InClassPath, UClass* InClass);
+	CONTENTBROWSERCLASSDATASOURCE_API FContentBrowserItemData CreateClassFileItem(UContentBrowserDataSource* InOwnerDataSource, const FName InVirtualPath, const FName InClassPath, UClass* InClass, const bool bIsFromPlugin);
 
 	CONTENTBROWSERCLASSDATASOURCE_API TSharedPtr<const FContentBrowserClassFolderItemDataPayload> GetClassFolderItemPayload(const UContentBrowserDataSource* InOwnerDataSource, const FContentBrowserItemData& InItem);
 
