Large perf improvement for the editor in a large project:

Reworked the ClassDataSource so that it doesn't require a full data refresh when a module is Loaded, Unloaded or Reloaded.
Deprecated the capacity for a ContentBrowser Data Source to request a full data refresh since this doesn't well scale on large projects.

Added a way to query a IReload for the classes that were reinstancited.

Fixed a bug in the FNativeClassHierarchy::GatherMatchingNodesForPaths function. It failed a getting a existing class node if a folder in its path had the same name has the class.

diff --git a/Engine/Plugins/Editor/ContentBrowser/ContentBrowserClassDataSource/Source/ContentBrowserClassDataSource/Private/ContentBrowserClassDataSource.cpp b/Engine/Plugins/Editor/ContentBrowser/ContentBrowserClassDataSource/Source/ContentBrowserClassDataSource/Private/ContentBrowserClassDataSource.cpp
--- a/Engine/Plugins/Editor/ContentBrowser/ContentBrowserClassDataSource/Source/ContentBrowserClassDataSource/Private/ContentBrowserClassDataSource.cpp
+++ b/Engine/Plugins/Editor/ContentBrowser/ContentBrowserClassDataSource/Source/ContentBrowserClassDataSource/Private/ContentBrowserClassDataSource.cpp
@@ -400,12 +400,12 @@ bool UContentBrowserClassDataSource::EnumerateItemsForObjects(const TArrayView<U
 		if (UClass* InClass = Cast<UClass>(InObject))
 		{
 			InternalPath.Reset();
-			if (NativeClassHierarchy->GetClassPath(InClass, InternalPath, NativeClassHierarchyGetClassPathCache))
+			if (NativeClassHierarchy->GetClassPath(InClass, InternalPath, NativeClassHierarchyGetClassPathCache.GameModules))
 			{
 				FName ClassPath = FName(InternalPath);
 				TSharedPtr<const FNativeClassHierarchyNode> ClassNode = NativeClassHierarchy->FindNode(ClassPath, ENativeClassHierarchyNodeType::Class);
 				const FContentBrowserItemPath ContentBrowserItemPath(InternalPath, EContentBrowserPathType::Internal);
-				if (!InCallback(ContentBrowserClassData::CreateClassFileItem(this, ContentBrowserItemPath.GetVirtualPathName(), ClassPath, InClass, ClassNode.IsValid() && ClassNode->LoadedFrom.IsSet())))
+				if (!InCallback(ContentBrowserClassData::CreateClassFileItem(this, ContentBrowserItemPath.GetVirtualPathName(), ClassPath, InClass, ClassNode->LoadedFrom.IsSet())))
 				{
 					return false;
 				}
@@ -594,12 +594,18 @@ bool UContentBrowserClassDataSource::GetClassPathsForCollections(TArrayView<cons
 }
 
 FContentBrowserItemData UContentBrowserClassDataSource::CreateClassFolderItem(const FName InFolderPath)
+{
+	TSharedPtr<const FNativeClassHierarchyNode> FolderNode = NativeClassHierarchy->FindNode(InFolderPath, ENativeClassHierarchyNodeType::Folder);
+
+	return CreateClassFolderItem(InFolderPath, FolderNode);
+}
+
+FContentBrowserItemData UContentBrowserClassDataSource::CreateClassFolderItem(const FName InFolderPath, const TSharedPtr<const FNativeClassHierarchyNode>& InFolderNode)
 {
 	FName VirtualizedPath;
 	TryConvertInternalPathToVirtual(InFolderPath, VirtualizedPath);
-	
-	TSharedPtr<const FNativeClassHierarchyNode> FolderNode = NativeClassHierarchy->FindNode(InFolderPath, ENativeClassHierarchyNodeType::Folder);
-	return ContentBrowserClassData::CreateClassFolderItem(this, VirtualizedPath, InFolderPath, FolderNode.IsValid() && FolderNode->LoadedFrom.IsSet());
+
+	return ContentBrowserClassData::CreateClassFolderItem(this, VirtualizedPath, InFolderPath, InFolderNode->LoadedFrom.IsSet());
 }
 
 FContentBrowserItemData UContentBrowserClassDataSource::CreateClassFileItem(UClass* InClass, FNativeClassHierarchyGetClassPathCache& InCache)
@@ -609,16 +615,22 @@ FContentBrowserItemData UContentBrowserClassDataSource::CreateClassFileItem(UCla
 	FName ClassPath;
 	{
 		FString ClassPathStr;
-		const bool bValidClassPath = NativeClassHierarchy->GetClassPath(InClass, ClassPathStr, InCache);
+		const bool bValidClassPath = NativeClassHierarchy->GetClassPath(InClass, ClassPathStr, InCache.GameModules);
 		checkf(bValidClassPath, TEXT("GetClassPath failed to return a result for '%s'"), *InClass->GetPathName());
 		ClassPath = *ClassPathStr;
 	}
 
+	TSharedPtr<const FNativeClassHierarchyNode> ClassNode = NativeClassHierarchy->FindNode(ClassPath, ENativeClassHierarchyNodeType::Class);
+
+	return CreateClassFileItem(ClassPath, ClassNode);
+}
+
+FContentBrowserItemData UContentBrowserClassDataSource::CreateClassFileItem(const FName InClassPath, const TSharedPtr<const FNativeClassHierarchyNode>& InClassNode)
+{
 	FName VirtualizedPath;
-	TryConvertInternalPathToVirtual(ClassPath, VirtualizedPath);
+	TryConvertInternalPathToVirtual(InClassPath, VirtualizedPath);
 
-	TSharedPtr<const FNativeClassHierarchyNode> ClassNode = NativeClassHierarchy->FindNode(ClassPath, ENativeClassHierarchyNodeType::Class);
-	return ContentBrowserClassData::CreateClassFileItem(this, VirtualizedPath, ClassPath, InClass, ClassNode.IsValid() && ClassNode->LoadedFrom.IsSet());
+	return ContentBrowserClassData::CreateClassFileItem(this, VirtualizedPath, InClassPath, InClassNode->Class, InClassNode->LoadedFrom.IsSet());
 }
 
 TSharedPtr<const FContentBrowserClassFolderItemDataPayload> UContentBrowserClassDataSource::GetClassFolderItemPayload(const FContentBrowserItemData& InItem) const
@@ -689,15 +701,50 @@ void UContentBrowserClassDataSource::ConditionalCreateNativeClassHierarchy()
 	if (!NativeClassHierarchy)
 	{
 		NativeClassHierarchy = MakeShared<FNativeClassHierarchy>();
-		NativeClassHierarchy->OnClassHierarchyUpdated().AddUObject(this, &UContentBrowserClassDataSource::ClassHierarchyUpdated);
+
+		NativeClassHierarchy->OnClassesAdded().AddUObject(this, &UContentBrowserClassDataSource::OnClassesAdded);
+		NativeClassHierarchy->OnClassesRemoved().AddUObject(this, &UContentBrowserClassDataSource::OnClassesRemoved);
+		NativeClassHierarchy->OnFoldersAdded().AddUObject(this, &UContentBrowserClassDataSource::OnFoldersAdded);
+		NativeClassHierarchy->OnFoldersRemoved().AddUObject(this, &UContentBrowserClassDataSource::OnFoldersRemoved);
+	}
+}
+
+void UContentBrowserClassDataSource::OnFoldersAdded(const TArrayView<TSharedRef<const FNativeClassHierarchyNode>> InFolders)
+{
+	NativeClassHierarchyGetClassPathCache.Reset();
+	SetVirtualPathTreeNeedsRebuild();
+
+	for (const TSharedRef<const FNativeClassHierarchyNode>& Folder : InFolders)
+	{
+		QueueItemDataUpdate(FContentBrowserItemDataUpdate::MakeItemAddedUpdate(CreateClassFolderItem(*Folder->EntryPath, Folder)));
 	}
 }
 
-void UContentBrowserClassDataSource::ClassHierarchyUpdated()
+void UContentBrowserClassDataSource::OnFoldersRemoved(const TArrayView<TSharedRef<const FNativeClassHierarchyNode>> InFolders)
 {
 	NativeClassHierarchyGetClassPathCache.Reset();
 	SetVirtualPathTreeNeedsRebuild();
-	NotifyItemDataRefreshed();
+
+	for (const TSharedRef<const FNativeClassHierarchyNode>& Folder : InFolders)
+	{
+		QueueItemDataUpdate(FContentBrowserItemDataUpdate::MakeItemRemovedUpdate(CreateClassFolderItem(*Folder->EntryPath, Folder)));
+	}
+}
+
+void UContentBrowserClassDataSource::OnClassesAdded(const TArrayView<TSharedRef<const FNativeClassHierarchyNode>> InClasses)
+{
+	for (const TSharedRef<const FNativeClassHierarchyNode>& Class : InClasses)
+	{
+		QueueItemDataUpdate(FContentBrowserItemDataUpdate::MakeItemAddedUpdate(CreateClassFileItem(*Class->EntryPath, Class)));
+	}
+}
+
+void UContentBrowserClassDataSource::OnClassesRemoved(const TArrayView<TSharedRef<const FNativeClassHierarchyNode>> InClasses)
+{
+	for (const TSharedRef<const FNativeClassHierarchyNode>& Class : InClasses)
+	{
+		QueueItemDataUpdate(FContentBrowserItemDataUpdate::MakeItemRemovedUpdate(CreateClassFileItem(*Class->EntryPath, Class)));
+	}
 }
 
 #undef LOCTEXT_NAMESPACE
diff --git a/Engine/Plugins/Editor/ContentBrowser/ContentBrowserClassDataSource/Source/ContentBrowserClassDataSource/Private/NativeClassHierarchy.cpp b/Engine/Plugins/Editor/ContentBrowser/ContentBrowserClassDataSource/Source/ContentBrowserClassDataSource/Private/NativeClassHierarchy.cpp
--- a/Engine/Plugins/Editor/ContentBrowser/ContentBrowserClassDataSource/Source/ContentBrowserClassDataSource/Private/NativeClassHierarchy.cpp
+++ b/Engine/Plugins/Editor/ContentBrowser/ContentBrowserClassDataSource/Source/ContentBrowserClassDataSource/Private/NativeClassHierarchy.cpp
@@ -1,16 +1,20 @@
 // Copyright Epic Games, Inc. All Rights Reserved.
 
 #include "NativeClassHierarchy.h"
-#include "Modules/ModuleManager.h"
-#include "UObject/UObjectIterator.h"
-#include "UObject/Package.h"
-#include "Misc/App.h"
-#include "Misc/PackageName.h"
-#include "Kismet2/KismetEditorUtilities.h"
-#include "SourceCodeNavigation.h"
+
 #include "Interfaces/IPluginManager.h"
 #include "Interfaces/IProjectManager.h"
+#include "Kismet2/KismetEditorUtilities.h"
+#include "Misc/App.h"
+#include "Misc/PackageName.h"
+#include "Modules/ModuleManager.h"
+#include "ProfilingDebugging/CpuProfilerTrace.h"
 #include "ProjectDescriptor.h"
+#include "SourceCodeNavigation.h"
+#include "String/ParseTokens.h"
+#include "UObject/Package.h"
+#include "UObject/Reload.h"
+#include "UObject/UObjectIterator.h"
 
 DEFINE_LOG_CATEGORY_STATIC(LogNativeClassHierarchy, Log, All);
 
@@ -52,14 +56,16 @@ FNativeClassHierarchy::FNativeClassHierarchy()
 	FModuleManager::Get().OnModulesChanged().AddRaw(this, &FNativeClassHierarchy::OnModulesChanged);
 
 	// Register to be notified of reloads
-	FCoreUObjectDelegates::ReloadCompleteDelegate.AddRaw(this, &FNativeClassHierarchy::OnReloadComplete);
+	FCoreUObjectDelegates::ReloadReinstancingCompleteDelegate.AddRaw(this, &FNativeClassHierarchy::OnReloadReinstancingComplete);
+	FCoreUObjectDelegates::ReloadAddedClassesDelegate.AddRaw(this, &FNativeClassHierarchy::OnReloadClassesAdded);
 }
 
 FNativeClassHierarchy::~FNativeClassHierarchy()
 {
 	FModuleManager::Get().OnModulesChanged().RemoveAll(this);
 
-	FCoreUObjectDelegates::ReloadCompleteDelegate.RemoveAll(this);
+	FCoreUObjectDelegates::ReloadReinstancingCompleteDelegate.RemoveAll(this);
+	FCoreUObjectDelegates::ReloadAddedClassesDelegate.RemoveAll(this);
 }
 
 TSharedPtr<const FNativeClassHierarchyNode> FNativeClassHierarchy::FindNode(const FName InClassPath, const ENativeClassHierarchyNodeType InType) const
@@ -286,7 +292,8 @@ void FNativeClassHierarchy::GatherMatchingNodesForPaths(const TArrayView<const F
 				// Try and find the node associated with this part of the path...
 				const FName ClassPathPartName = *ClassPathPart;
 
-				if (InType == ENativeClassHierarchyNodeType::Class && ClassPathPart == ClassPathParts.Last())
+				// Compare the address of the strings to see if it is the last part
+				if (InType == ENativeClassHierarchyNodeType::Class && &ClassPathPart == &ClassPathParts.Last())
 				{
 
 					CurrentNode = (CurrentNode.IsValid()) ? CurrentNode->Children.FindRef(FNativeClassHierarchyNodeKey(ClassPathPartName, ENativeClassHierarchyNodeType::Class)) : RootNodes.FindRef(ClassPathPartName);
@@ -313,7 +320,11 @@ void FNativeClassHierarchy::GatherMatchingNodesForPaths(const TArrayView<const F
 
 void FNativeClassHierarchy::PopulateHierarchy()
 {
-	FAddClassMetrics AddClassMetrics;
+	TRACE_CPUPROFILER_EVENT_SCOPE(FNativeClassHierarchy::PopulateHierarchy);
+
+	double StartTime = FPlatformTime::Seconds();
+
+	FClassChanges ClassChanges;
 
 	RootNodes.Empty();
 
@@ -322,17 +333,27 @@ void FNativeClassHierarchy::PopulateHierarchy()
 	for(TObjectIterator<UClass> ClassIt; ClassIt; ++ClassIt)
 	{
 		UClass* const CurrentClass = *ClassIt;
-		AddClass(CurrentClass, GameModules, AddClassMetrics);
+		AddClass(CurrentClass, GameModules, ClassChanges);
 	}
 
-	UE_LOG(LogNativeClassHierarchy, Verbose, TEXT("Native class hierarchy populated in %0.4f seconds. Added %d classes and %d folders."), FPlatformTime::Seconds() - AddClassMetrics.StartTime, AddClassMetrics.NumClassesAdded, AddClassMetrics.NumFoldersAdded);
+	UE_LOG(LogNativeClassHierarchy, Verbose, TEXT("Native class hierarchy populated in %0.4f seconds. Added %d classes and %d folders."), FPlatformTime::Seconds() - StartTime, ClassChanges.ClassesModified.Num(), ClassChanges.FoldersModified.Num());
 
-	ClassHierarchyUpdatedDelegate.Broadcast();
+	if (!ClassChanges.FoldersModified.IsEmpty())
+	{
+		FoldersAddedDelegate.Broadcast(ClassChanges.FoldersModified);
+	}
+	if (!ClassChanges.ClassesModified.IsEmpty())
+	{
+		ClassesAddedDelegate.Broadcast(ClassChanges.ClassesModified);
+	}
 }
 
 void FNativeClassHierarchy::AddClassesForModule(const FName& InModuleName)
 {
-	FAddClassMetrics AddClassMetrics;
+	TRACE_CPUPROFILER_EVENT_SCOPE(FNativeClassHierarchy::AddClassesForModule);
+
+	double StartTime = FPlatformTime::Seconds();
+	FClassChanges ClassChanges;
 
 	// Find the class package for this module
 	UPackage* const ClassPackage = FindPackage(nullptr, *(FString("/Script/") + InModuleName.ToString()));
@@ -345,38 +366,63 @@ void FNativeClassHierarchy::AddClassesForModule(const FName& InModuleName)
 	
 	TArray<UObject*> PackageObjects;
 	GetObjectsWithOuter(ClassPackage, PackageObjects, false);
+	ClassChanges.ClassesModified.Reserve(PackageObjects.Num());
+
 	for(UObject* Object : PackageObjects)
 	{
 		UClass* const CurrentClass = Cast<UClass>(Object);
 		if(CurrentClass)
 		{
-			AddClass(CurrentClass, GameModules, AddClassMetrics);
+			AddClass(CurrentClass, GameModules, ClassChanges);
 		}
 	}
 
-	UE_LOG(LogNativeClassHierarchy, Verbose, TEXT("Native class hierarchy updated for '%s' in %0.4f seconds. Added %d classes and %d folders."), *InModuleName.ToString(), FPlatformTime::Seconds() - AddClassMetrics.StartTime, AddClassMetrics.NumClassesAdded, AddClassMetrics.NumFoldersAdded);
+	UE_LOG(LogNativeClassHierarchy, Verbose, TEXT("Native class hierarchy updated for '%s' in %0.4f seconds. Added %d classes and %d folders."), *InModuleName.ToString(), FPlatformTime::Seconds() - StartTime, ClassChanges.ClassesModified.Num(), ClassChanges.FoldersModified.Num());
 
-	ClassHierarchyUpdatedDelegate.Broadcast();
+	if (!ClassChanges.FoldersModified.IsEmpty())
+	{
+		FoldersAddedDelegate.Broadcast(ClassChanges.FoldersModified);
+	}
+	if (!ClassChanges.ClassesModified.IsEmpty())
+	{
+		ClassesAddedDelegate.Broadcast(ClassChanges.ClassesModified);
+	}
 }
 
 void FNativeClassHierarchy::RemoveClassesForModule(const FName& InModuleName)
 {
+	TRACE_CPUPROFILER_EVENT_SCOPE(FNativeClassHierarchy::RemoveClassesForModule);
+
 	// Modules always exist directly under a root
 	for(const auto& RootNode : RootNodes)
 	{
 		TSharedPtr<FNativeClassHierarchyNode> ModuleNode = RootNode.Value->Children.FindRef(FNativeClassHierarchyNodeKey(InModuleName, ENativeClassHierarchyNodeType::Folder));
 		if(ModuleNode.IsValid())
 		{
-			// Remove this module from its root
-			RootNode.Value->Children.Remove(FNativeClassHierarchyNodeKey(InModuleName, ENativeClassHierarchyNodeType::Folder));
+			FClassChanges ClassChanges;
 
 			// If this module was the only child of this root, then we need to remove the root as well
-			if(RootNode.Value->Children.Num() == 0)
+			if(RootNode.Value->Children.Num() == 1)
 			{
+				PopulateClassChanges(RootNode.Value, ClassChanges);
 				RootNodes.Remove(RootNode.Key);
 			}
+			else
+			{
+				// Remove this module from its root
+				TSharedPtr<FNativeClassHierarchyNode> RemovedNode;
+				RootNode.Value->Children.RemoveAndCopyValue(FNativeClassHierarchyNodeKey(InModuleName, ENativeClassHierarchyNodeType::Folder), RemovedNode);
+				PopulateClassChanges(RemovedNode, ClassChanges);
+			}
 
-			ClassHierarchyUpdatedDelegate.Broadcast();
+			if (!ClassChanges.ClassesModified.IsEmpty())
+			{
+				ClassesRemovedDelegate.Broadcast(ClassChanges.ClassesModified);
+			}
+			if (!ClassChanges.FoldersModified.IsEmpty())
+			{
+				FoldersRemovedDelegate.Broadcast(ClassChanges.FoldersModified);
+			}
 
 			// We've found the module - break
 			break;
@@ -384,7 +430,7 @@ void FNativeClassHierarchy::RemoveClassesForModule(const FName& InModuleName)
 	}
 }
 
-void FNativeClassHierarchy::AddClass(UClass* InClass, const TSet<FName>& InGameModules, FAddClassMetrics& AddClassMetrics)
+void FNativeClassHierarchy::AddClass(UClass* InClass, const TSet<FName>& InGameModules, FClassChanges& ClassChanges)
 {
 	// Ignore deprecated and temporary classes
 	if(InClass->HasAnyClassFlags(CLASS_Deprecated | CLASS_NewerVersionExists) || FKismetEditorUtilities::IsClassABlueprintSkeleton(InClass))
@@ -419,7 +465,7 @@ void FNativeClassHierarchy::AddClass(UClass* InClass, const TSet<FName>& InGameM
 	{
 		RootNode = FNativeClassHierarchyNode::MakeFolderEntry(RootNodeName, TEXT("/") + RootNodeName.ToString(), WhereLoadedFrom);
 		RootNode->LoadedFrom = WhereLoadedFrom;
-		++AddClassMetrics.NumFoldersAdded;
+		ClassChanges.FoldersModified.Add(RootNode.ToSharedRef());
 	}
 
 	// Split the class path and ensure we have nodes for each part
@@ -433,14 +479,94 @@ void FNativeClassHierarchy::AddClass(UClass* InClass, const TSet<FName>& InGameM
 		if(!ChildNode.IsValid())
 		{
 			ChildNode = FNativeClassHierarchyNode::MakeFolderEntry(HierarchyPathPartName, CurrentNode->EntryPath + TEXT("/") + HierarchyPathPart, WhereLoadedFrom);
-			++AddClassMetrics.NumFoldersAdded;
+			ClassChanges.FoldersModified.Add(ChildNode.ToSharedRef());
 		}
 		CurrentNode = ChildNode;
 	}
 
 	// Now add the final entry for the class
-	CurrentNode->AddChild(FNativeClassHierarchyNode::MakeClassEntry(InClass, ClassModuleName, ClassModuleRelativePath, CurrentNode->EntryPath + TEXT("/") + InClass->GetName(), WhereLoadedFrom));
-	++AddClassMetrics.NumClassesAdded;
+	TSharedRef<FNativeClassHierarchyNode> ClassNode = FNativeClassHierarchyNode::MakeClassEntry(InClass, ClassModuleName, ClassModuleRelativePath, CurrentNode->EntryPath + TEXT("/") + InClass->GetName(), WhereLoadedFrom);
+	ClassChanges.ClassesModified.Add(ClassNode);
+	CurrentNode->AddChild(MoveTemp(ClassNode));
+}
+
+void FNativeClassHierarchy::RemoveClass(UClass* InClass, TSet<FName>& InGameModules, FClassChanges& ClassChanges)
+{
+	FString OutClassPath;
+	constexpr bool bIncludeClass = false;
+	if (!GetClassPath(InClass, OutClassPath, InGameModules, bIncludeClass))
+	{
+		return;
+	}
+	
+	// Most items aren't deep in the three but having a little buffer still help
+	TArray<TSharedRef<FNativeClassHierarchyNode>, TInlineAllocator<10>> FoldersVisited;
+	bool bHasProcessRootNodes = false;
+
+	auto Visitor = [&FoldersVisited, &bHasProcessRootNodes, this](FStringView Token)
+		{
+			if (!bHasProcessRootNodes)
+			{
+				if (TSharedPtr<FNativeClassHierarchyNode>* NodePtr = RootNodes.Find(FName(Token)))
+				{
+					FoldersVisited.Add(NodePtr->ToSharedRef());
+				}
+
+				bHasProcessRootNodes = true;
+			}
+			else if (!FoldersVisited.IsEmpty())
+			{
+				if (TSharedPtr<FNativeClassHierarchyNode>* NodePtr = FoldersVisited.Last()->Children.Find(FNativeClassHierarchyNodeKey(FName(Token), ENativeClassHierarchyNodeType::Folder)))
+				{
+					FoldersVisited.Add(NodePtr->ToSharedRef());
+				}
+			}
+
+		};
+
+	UE::String::ParseTokens(FStringView(OutClassPath), TEXT('/'), TFunctionRef<void (FStringView)>(Visitor));
+
+	if (!FoldersVisited.IsEmpty())
+	{
+		TSharedPtr<FNativeClassHierarchyNode> RemovedClass;
+		bool bHasRemovedNode = false;
+
+		bHasRemovedNode = FoldersVisited.Last()->Children.RemoveAndCopyValue(FNativeClassHierarchyNodeKey(InClass->GetFName(), ENativeClassHierarchyNodeType::Class), RemovedClass);
+		ClassChanges.ClassesModified.Add(RemovedClass.ToSharedRef());
+
+		if (bHasRemovedNode)
+		{
+			for (int Index = FoldersVisited.Num() - 2; Index >= 0; --Index)
+			{
+				const TSharedRef<FNativeClassHierarchyNode> ToRemoveIfEmpty = FoldersVisited[Index + 1];
+				if (ToRemoveIfEmpty->Children.IsEmpty())
+				{
+					ClassChanges.FoldersModified.Add(ToRemoveIfEmpty);
+					FoldersVisited[Index]->Children.Remove(FNativeClassHierarchyNodeKey(ToRemoveIfEmpty->EntryName, ENativeClassHierarchyNodeType::Folder));
+				}
+				else
+				{
+					bHasRemovedNode = false;
+				}
+
+				if (!bHasRemovedNode)
+				{
+					break;
+				}
+			}
+
+			// If we removed all the node until the root. Check if we should remove the root also
+			if (bHasRemovedNode)
+			{
+				const TSharedRef<FNativeClassHierarchyNode>& ToRemoveIfEmpty = FoldersVisited[0];
+				if (ToRemoveIfEmpty->Children.IsEmpty())
+				{
+					ClassChanges.FoldersModified.Add(ToRemoveIfEmpty);
+					RootNodes.Remove(ToRemoveIfEmpty->EntryName);
+				}
+			}
+		}
+	}
 }
 
 bool FNativeClassHierarchy::GetFileSystemPath(const FString& InClassPath, FString& OutFileSystemPath) const
@@ -482,7 +608,7 @@ bool FNativeClassHierarchy::GetFileSystemPath(const FString& InClassPath, FStrin
 	return false;
 }
 
-bool FNativeClassHierarchy::GetClassPath(const UClass* InClass, FString& OutClassPath, FNativeClassHierarchyGetClassPathCache& InCache, const bool bIncludeClassName) const
+bool FNativeClassHierarchy::GetClassPath(const UClass* InClass, FString& OutClassPath, TSet<FName>& InGameModuleCache, const bool bIncludeClassName) const
 {
 	const FName ClassModuleName = GetClassModuleName(InClass);
 	if(ClassModuleName.IsNone())
@@ -497,14 +623,14 @@ bool FNativeClassHierarchy::GetClassPath(const UClass* InClass, FString& OutClas
 		return false;
 	}
 
-	if (InCache.GameModules.Num() == 0)
+	if (InGameModuleCache.Num() == 0)
 	{
-		InCache.GameModules = GetGameModules();
+		InGameModuleCache = GetGameModules();
 	}
 
 	// Work out which root this class should go under
 	TOptional<EPluginLoadedFrom> WhereLoadedFrom;
-	const FName RootNodeName = GetClassPathRootForModule(ClassModuleName, InCache.GameModules, WhereLoadedFrom);
+	const FName RootNodeName = GetClassPathRootForModule(ClassModuleName, InGameModuleCache, WhereLoadedFrom);
 
 	// Work out the final path to this class within the hierarchy (which isn't the same as the path on disk)
 	const FString ClassModuleRelativePath = ClassModuleRelativeIncludePath.Left(ClassModuleRelativeIncludePath.Find(TEXT("/"), ESearchCase::CaseSensitive, ESearchDir::FromEnd));
@@ -540,9 +666,87 @@ void FNativeClassHierarchy::OnModulesChanged(FName InModuleName, EModuleChangeRe
 	}
 }
 
-void FNativeClassHierarchy::OnReloadComplete(EReloadCompleteReason Reason)
+void FNativeClassHierarchy::OnReloadReinstancingComplete()
 {
-	PopulateHierarchy();
+	TRACE_CPUPROFILER_EVENT_SCOPE(FNativeClassHierarchy::OnReloadReinstancingComplete);
+
+	if (IReload* ActiveReload = GetActiveReloadInterface())
+	{
+		if (const TMap<UClass*, UClass*>* ReinstancedClassesPtr = ActiveReload->GetReinstancedClasses())
+		{
+			const TMap<UClass*, UClass*>& ReinstancedClasses = *ReinstancedClassesPtr;
+			if (!ReinstancedClasses.IsEmpty())
+			{
+				FClassChanges ClassChanges;
+
+				TSet<FName> GameModules = GetGameModules();
+
+				// Add the new classes first
+				for (const TPair<UClass*, UClass*>& OldAndNew : ReinstancedClasses)
+				{
+					
+
+					if (OldAndNew.Value)
+					{
+						// Adding the new class will replace old entry
+						AddClass(OldAndNew.Value, GameModules, ClassChanges);
+					}
+				}
+
+				if (!ClassChanges.FoldersModified.IsEmpty())
+				{
+					FoldersAddedDelegate.Broadcast(ClassChanges.FoldersModified);
+				}
+				if (!ClassChanges.ClassesModified.IsEmpty())
+				{
+					ClassesAddedDelegate.Broadcast(ClassChanges.ClassesModified);
+				}
+
+				// Reset the class changes tracking
+				ClassChanges.Reset();
+
+				// Then remove the old classes
+				for (const TPair<UClass*, UClass*>& OldAndNew : ReinstancedClasses)
+				{
+					if (OldAndNew.Key && !OldAndNew.Value)
+					{
+						RemoveClass(OldAndNew.Key, GameModules, ClassChanges);
+					}
+					}
+
+				if (!ClassChanges.ClassesModified.IsEmpty())
+				{
+					ClassesRemovedDelegate.Broadcast(ClassChanges.ClassesModified);
+				}
+				if (!ClassChanges.FoldersModified.IsEmpty())
+				{
+					FoldersRemovedDelegate.Broadcast(ClassChanges.FoldersModified);
+				}
+			}
+		}
+	}
+}
+
+void FNativeClassHierarchy::OnReloadClassesAdded(const TArray<UClass*>& InAddedClasses)
+{
+	TRACE_CPUPROFILER_EVENT_SCOPE(FNativeClassHierarchy::InAddedClasses);
+
+	FClassChanges ClassChanges;
+	TSet<FName> GameModules = GetGameModules();
+
+	for (UClass* AddedClass : InAddedClasses)
+	{
+		AddClass(AddedClass, GameModules, ClassChanges);
+	}
+
+	if (!ClassChanges.FoldersModified.IsEmpty())
+	{
+		FoldersAddedDelegate.Broadcast(ClassChanges.FoldersModified);
+	}
+	if (!ClassChanges.ClassesModified.IsEmpty())
+	{
+		ClassesAddedDelegate.Broadcast(ClassChanges.ClassesModified);
+	}
 }
 
 FName FNativeClassHierarchy::GetClassModuleName(const UClass* InClass)
@@ -599,3 +803,20 @@ TSet<FName> FNativeClassHierarchy::GetGameModules()
 
 	return GameModules;
 }
+
+void FNativeClassHierarchy::PopulateClassChanges(const TSharedPtr<FNativeClassHierarchyNode>& InNode, FClassChanges& OutClassChanges)
+{
+	for (const TPair<FNativeClassHierarchyNodeKey, TSharedPtr<FNativeClassHierarchyNode>>& Childs : InNode->Children)
+	{
+		PopulateClassChanges(Childs.Value, OutClassChanges);
+	}
+
+	if (InNode->Type == ENativeClassHierarchyNodeType::Folder)
+	{
+		OutClassChanges.FoldersModified.Add(InNode.ToSharedRef());
+	}
+	else if (InNode->Type == ENativeClassHierarchyNodeType::Class)
+	{
+		OutClassChanges.ClassesModified.Add(InNode.ToSharedRef());
+	}
+}
diff --git a/Engine/Plugins/Editor/ContentBrowser/ContentBrowserClassDataSource/Source/ContentBrowserClassDataSource/Private/NativeClassHierarchy.h b/Engine/Plugins/Editor/ContentBrowser/ContentBrowserClassDataSource/Source/ContentBrowserClassDataSource/Private/NativeClassHierarchy.h
--- a/Engine/Plugins/Editor/ContentBrowser/ContentBrowserClassDataSource/Source/ContentBrowserClassDataSource/Private/NativeClassHierarchy.h
+++ b/Engine/Plugins/Editor/ContentBrowser/ContentBrowserClassDataSource/Source/ContentBrowserClassDataSource/Private/NativeClassHierarchy.h
@@ -6,7 +6,6 @@
 #include "HAL/PlatformTime.h"
 
 enum class EPluginLoadedFrom;
-enum class EReloadCompleteReason;
 
 /**
  * Type of hierarchy node
@@ -156,10 +155,31 @@ public:
 	FNativeClassHierarchy();
 	~FNativeClassHierarchy();
 
-	/** Get the delegate called when classes are added or removed from this class hierarchy */
-	FSimpleMulticastDelegate& OnClassHierarchyUpdated()
+
+	DECLARE_MULTICAST_DELEGATE_OneParam(FOnNodesChanged, const TArrayView<TSharedRef<const FNativeClassHierarchyNode>>)
+
+	/** Get the delegate called when folders are added to this class hierarchy */
+	FOnNodesChanged& OnFoldersAdded()
+	{
+		return FoldersAddedDelegate;
+	}
+
+	/** Get the delegate called when classes are added to this class hierarchy */
+	FOnNodesChanged& OnClassesAdded()
+	{
+		return ClassesAddedDelegate;
+	}
+
+	/** Get the delegate called when folders are removed to this class hierarchy */
+	FOnNodesChanged& OnFoldersRemoved()
+	{
+		return FoldersRemovedDelegate;
+	}
+
+	/** Get the delegate called when classes are removed to this class hierarchy */
+	FOnNodesChanged& OnClassesRemoved()
 	{
-		return ClassHierarchyUpdatedDelegate;
+		return ClassesRemovedDelegate;
 	}
 
 	/**
@@ -241,7 +261,7 @@ public:
 	 * 
 	 * @return true if the class path could be resolved and OutClassPath was filled in, false otherwise
 	 */
-	bool GetClassPath(const UClass* InClass, FString& OutClassPath, FNativeClassHierarchyGetClassPathCache& InCache, const bool bIncludeClassName = true) const;
+	bool GetClassPath(const UClass* InClass, FString& OutClassPath, TSet<FName>& InGameModuleCache, const bool bIncludeClassName = true) const;
 
 	/**
 	 * Test if root node passes given rules
@@ -254,18 +274,20 @@ public:
 	bool RootNodePassesFilter(const FName InRootName, const bool bIncludeEngineClasses, const bool bIncludePluginClasses) const;
 
 private:
-	struct FAddClassMetrics
+	struct FClassChanges
 	{
-		FAddClassMetrics()
-			: StartTime(FPlatformTime::Seconds())
-			, NumClassesAdded(0)
-			, NumFoldersAdded(0)
+		FClassChanges()
+		{
+		}
+
+		void Reset()
 		{
+			ClassesModified.Reset();
+			FoldersModified.Reset();
 		}
 
-		double StartTime;
-		int32 NumClassesAdded;
-		int32 NumFoldersAdded;
+		TArray<TSharedRef<const FNativeClassHierarchyNode>> ClassesModified;
+		TArray<TSharedRef<const FNativeClassHierarchyNode>> FoldersModified;
 	};
 
 	/**
@@ -339,9 +361,18 @@ private:
 	 *
 	 * @param InClass - The class that is to be added
 	 * @param InGameModules - The list of modules that belong to the "/Classes_Game" root
-	 * @param AddClassMetrics - Metrics to update as new classes and folders are added (used to report population performance)
+	 * @param FClassChanges - Track the new classes and folders that were added.
 	 */
-	void AddClass(UClass* InClass, const TSet<FName>& InGameModules, FAddClassMetrics& AddClassMetrics);
+	void AddClass(UClass* InClass, const TSet<FName>& InGameModules, FClassChanges& ClassChanges);
+
+	/**
+	 * Remove a single class to the known class hierarchy, removing any empty folders left
+	 *
+	 * @param InClass - The class that is to be removed
+	 * @param InGameModules - The list of modules that belong to the "/Classes_Game" root
+	 * @param FClassChanges - Track the classes and folders that were removed.
+	 */
+	void RemoveClass(UClass* InClass, TSet<FName>& InGameModules, FClassChanges& ClassChanges);
 
 	/**
 	 * Called when we're notified that a module has changed status
@@ -352,11 +383,14 @@ private:
 	void OnModulesChanged(FName InModuleName, EModuleChangeReason InModuleChangeReason);
 
 	/**
-	 * Called when we're notified that a module has been reloaded
-	 *
-	 * @param Reason - The reason why the module was reloaded
+	 * Called when we're notified that a reloaded module has reinstanced some classes.
 	 */
-	void OnReloadComplete(EReloadCompleteReason Reason);
+	void OnReloadReinstancingComplete();
+
+	/**
+	 * Called when we're notified that a reloaded module has some new classes.
+	 */
+	 void OnReloadClassesAdded(const TArray<UClass*>& InAddedClasses);
 
 	/**
 	 * Given a class, work out which module it belongs to
@@ -383,9 +417,20 @@ private:
 	 */
 	static TSet<FName> GetGameModules();
 
+	void PopulateClassChanges(const TSharedPtr<FNativeClassHierarchyNode>& InNode, FClassChanges& OutClassChanges);
+
 	/** Root level nodes corresponding to the root folders used by the Content Browser, eg) Classes_Engine, Classes_Game, etc */
 	TMap<FName, TSharedPtr<FNativeClassHierarchyNode>> RootNodes;
 
-	/** Delegate called when the class hierarchy is updated */
-	FSimpleMulticastDelegate ClassHierarchyUpdatedDelegate;
+	/** Delegate called when the class hierarchy as added some new folders */
+	FOnNodesChanged FoldersAddedDelegate;
+
+	/** Delegate called when the class hierarchy as added some new classes */
+	FOnNodesChanged ClassesAddedDelegate;
+
+	/** Delegate called when the class hierarchy as removed some folders */
+	FOnNodesChanged FoldersRemovedDelegate;
+
+	/** Delegate called when the class hierarchy as removed some classes */
+	FOnNodesChanged ClassesRemovedDelegate;
 };
diff --git a/Engine/Plugins/Editor/ContentBrowser/ContentBrowserClassDataSource/Source/ContentBrowserClassDataSource/Public/ContentBrowserClassDataSource.h b/Engine/Plugins/Editor/ContentBrowser/ContentBrowserClassDataSource/Source/ContentBrowserClassDataSource/Public/ContentBrowserClassDataSource.h
--- a/Engine/Plugins/Editor/ContentBrowser/ContentBrowserClassDataSource/Source/ContentBrowserClassDataSource/Public/ContentBrowserClassDataSource.h
+++ b/Engine/Plugins/Editor/ContentBrowser/ContentBrowserClassDataSource/Source/ContentBrowserClassDataSource/Public/ContentBrowserClassDataSource.h
@@ -92,6 +92,10 @@ private:
 
 	FContentBrowserItemData CreateClassFileItem(UClass* InClass, FNativeClassHierarchyGetClassPathCache& InCache);
 
+	FContentBrowserItemData CreateClassFolderItem(const FName InFolderPath, const TSharedPtr<const FNativeClassHierarchyNode>& InFolderNode);
+
+	FContentBrowserItemData CreateClassFileItem(const FName InClassPath, const TSharedPtr<const FNativeClassHierarchyNode>& InClassNode);
+
 	TSharedPtr<const FContentBrowserClassFolderItemDataPayload> GetClassFolderItemPayload(const FContentBrowserItemData& InItem) const;
 
 	TSharedPtr<const FContentBrowserClassFileItemDataPayload> GetClassFileItemPayload(const FContentBrowserItemData& InItem) const;
@@ -102,7 +106,13 @@ private:
 
 	void ConditionalCreateNativeClassHierarchy();
 
-	void ClassHierarchyUpdated();
+	void OnFoldersAdded(const TArrayView<TSharedRef<const FNativeClassHierarchyNode>> InFolders);
+
+	void OnFoldersRemoved(const TArrayView<TSharedRef<const FNativeClassHierarchyNode>> InFolders);
+
+	void OnClassesAdded(const TArrayView<TSharedRef<const FNativeClassHierarchyNode>> InClasses);
+
+	void OnClassesRemoved(const TArrayView<TSharedRef<const FNativeClassHierarchyNode>> InClasses);
 
 	TSharedPtr<FNativeClassHierarchy> NativeClassHierarchy;
 	FNativeClassHierarchyGetClassPathCache NativeClassHierarchyGetClassPathCache;
diff --git a/Engine/Source/Developer/Windows/LiveCoding/Private/LiveCodingModule.cpp b/Engine/Source/Developer/Windows/LiveCoding/Private/LiveCodingModule.cpp
--- a/Engine/Source/Developer/Windows/LiveCoding/Private/LiveCodingModule.cpp
+++ b/Engine/Source/Developer/Windows/LiveCoding/Private/LiveCodingModule.cpp
@@ -139,6 +139,11 @@ public:
 		return CDO;
 	}
 
+	virtual const TMap<UClass*, UClass*>* GetReinstancedClasses() override
+	{
+		return nullptr;
+	}
+
 	bool HasReinstancingOccurred() const
 	{
 		return bHasReinstancingOccurred;
diff --git a/Engine/Source/Editor/ContentBrowser/Private/SPathView.cpp b/Engine/Source/Editor/ContentBrowser/Private/SPathView.cpp
--- a/Engine/Source/Editor/ContentBrowser/Private/SPathView.cpp
+++ b/Engine/Source/Editor/ContentBrowser/Private/SPathView.cpp
@@ -45,6 +45,7 @@
 #include "Misc/StringBuilder.h"
 #include "Modules/ModuleManager.h"
 #include "PathViewTypes.h"
+#include "ProfilingDebugging/CpuProfilerTrace.h"
 #include "Settings/ContentBrowserSettings.h"
 #include "SlotBase.h"
 #include "SourcesData.h"
@@ -1518,6 +1519,8 @@ FText SPathView::GetHighlightText() const
 
 void SPathView::Populate(const bool bIsRefreshingFilter)
 {
+	TRACE_CPUPROFILER_EVENT_SCOPE(SPathView::Populate);
+
 	// Update the list of expanded path before removing the items
 	UpdateLastExpandedPathsIfDirty();
 
@@ -2303,6 +2306,8 @@ SFavoritePathView::~SFavoritePathView()
 
 void SFavoritePathView::Populate(const bool bIsRefreshingFilter)
 {
+	TRACE_CPUPROFILER_EVENT_SCOPE(SFavoritePathView::Populate);
+
 	// Don't allow the selection changed delegate to be fired here
 	FScopedPreventTreeItemChangedDelegate DelegatePrevention(SharedThis(this));
 
diff --git a/Engine/Source/Editor/ContentBrowserData/Private/ContentBrowserDataSubsystem.cpp b/Engine/Source/Editor/ContentBrowserData/Private/ContentBrowserDataSubsystem.cpp
--- a/Engine/Source/Editor/ContentBrowserData/Private/ContentBrowserDataSubsystem.cpp
+++ b/Engine/Source/Editor/ContentBrowserData/Private/ContentBrowserDataSubsystem.cpp
@@ -755,6 +755,8 @@ void UContentBrowserDataSubsystem::Tick(const float InDeltaTime)
 
 	if (bPendingItemDataRefreshedNotification)
 	{
+		TRACE_CPUPROFILER_EVENT_SCOPE(UContentBrowserDataSubsystem::BroadcastItemDataRefreshed);
+
 		bPendingItemDataRefreshedNotification = false;
 		bHasIgnoredItemUpdates = false;
 		PendingUpdates.Empty();
diff --git a/Engine/Source/Editor/ContentBrowserData/Public/ContentBrowserDataSource.h b/Engine/Source/Editor/ContentBrowserData/Public/ContentBrowserDataSource.h
--- a/Engine/Source/Editor/ContentBrowserData/Public/ContentBrowserDataSource.h
+++ b/Engine/Source/Editor/ContentBrowserData/Public/ContentBrowserDataSource.h
@@ -864,6 +864,7 @@ protected:
 	/**
 	 * Notify a wholesale item data update, for data sources that can't provide delta-updates.
 	 */
+	 UE_DEPRECATED(5.5, "No data source should be able request a full data refresh. On large projects this can quickly destroy the editor performance. To avoid that all data sources should support queueing the required incremental data updates for their items. Look for uses of the function QueueItemDataUpdate for examples.")
 	void NotifyItemDataRefreshed();
 
 	/**
diff --git a/Engine/Source/Editor/UnrealEd/Public/Kismet2/ReloadUtilities.h b/Engine/Source/Editor/UnrealEd/Public/Kismet2/ReloadUtilities.h
--- a/Engine/Source/Editor/UnrealEd/Public/Kismet2/ReloadUtilities.h
+++ b/Engine/Source/Editor/UnrealEd/Public/Kismet2/ReloadUtilities.h
@@ -44,6 +44,10 @@ public:
 	UNREALED_API virtual void Reinstance() override;
 	UNREALED_API virtual UObject* GetReinstancedCDO(UObject* CDO) override;
 	UNREALED_API virtual const UObject* GetReinstancedCDO(const UObject* CDO) override;
+	UNREALED_API virtual const TMap<UClass*, UClass*>* GetReinstancedClasses() override
+	{
+		return &ReinstancedClasses;
+	}
 
 	/**
 	 * If you wish to reuse the same reload object, invoke this method to reset the state
diff --git a/Engine/Source/Runtime/CoreUObject/Public/UObject/Reload.h b/Engine/Source/Runtime/CoreUObject/Public/UObject/Reload.h
--- a/Engine/Source/Runtime/CoreUObject/Public/UObject/Reload.h
+++ b/Engine/Source/Runtime/CoreUObject/Public/UObject/Reload.h
@@ -60,6 +60,12 @@ public:
 	 */
 	virtual UObject* GetReinstancedCDO(UObject* CDO) = 0;
 	virtual const UObject* GetReinstancedCDO(const UObject* CDO) = 0;
+
+	/**
+	 * Return an optional map of the reinstanced classes where the old class is the key and the new class is the value.
+	 * Note: the returned value can be null.
+	 */
+	 virtual const TMap<UClass*, UClass*>* GetReinstancedClasses() = 0;
 };
 
 /**
