Fix issue with content browser now properly detecting new classes from live coding.

#rnx

diff --git a/Engine/Plugins/Editor/ContentBrowser/ContentBrowserClassDataSource/Source/ContentBrowserClassDataSource/Private/NativeClassHierarchy.cpp b/Engine/Plugins/Editor/ContentBrowser/ContentBrowserClassDataSource/Source/ContentBrowserClassDataSource/Private/NativeClassHierarchy.cpp
--- a/Engine/Plugins/Editor/ContentBrowser/ContentBrowserClassDataSource/Source/ContentBrowserClassDataSource/Private/NativeClassHierarchy.cpp
+++ b/Engine/Plugins/Editor/ContentBrowser/ContentBrowserClassDataSource/Source/ContentBrowserClassDataSource/Private/NativeClassHierarchy.cpp
@@ -672,57 +672,59 @@ void FNativeClassHierarchy::OnReloadReinstancingComplete()
 
 	if (IReload* ActiveReload = GetActiveReloadInterface())
 	{
-		if (const TMap<UClass*, UClass*>* ReinstancedClassesPtr = ActiveReload->GetReinstancedClasses())
+		const TMap<UClass*, UClass*>& ReinstancedClasses = ActiveReload->GetReinstancedClasses();
+		const TArray<UClass*>& NewClasses = ActiveReload->GetNewClasses();
+		if (!ReinstancedClasses.IsEmpty() || !NewClasses.IsEmpty())
 		{
-			const TMap<UClass*, UClass*>& ReinstancedClasses = *ReinstancedClassesPtr;
-			if (!ReinstancedClasses.IsEmpty())
-			{
-				FClassChanges ClassChanges;
-
-				TSet<FName> GameModules = GetGameModules();
+			FClassChanges ClassChanges;
 
-				// Add the new classes first
-				for (const TPair<UClass*, UClass*>& OldAndNew : ReinstancedClasses)
-				{
-					
+			TSet<FName> GameModules = GetGameModules();
 
-					if (OldAndNew.Value)
-					{
-						// Adding the new class will replace old entry
-						AddClass(OldAndNew.Value, GameModules, ClassChanges);
-					}
-				}
+			// Add new classes
+			for (UClass* NewClass : NewClasses)
+			{
+				AddClass(NewClass, GameModules, ClassChanges);
+			}
 
-				if (!ClassChanges.FoldersModified.IsEmpty())
-				{
-					FoldersAddedDelegate.Broadcast(ClassChanges.FoldersModified);
-				}
-				if (!ClassChanges.ClassesModified.IsEmpty())
+			// Replace re-instanced classes
+			for (const TPair<UClass*, UClass*>& OldAndNew : ReinstancedClasses)
+			{
+				if (OldAndNew.Value)
 				{
-					ClassesAddedDelegate.Broadcast(ClassChanges.ClassesModified);
+					// Adding the re-instanced class will replace old entry
+					AddClass(OldAndNew.Value, GameModules, ClassChanges);
 				}
+			}
 
-				// Reset the class changes tracking
-				ClassChanges.Reset();
+			if (!ClassChanges.FoldersModified.IsEmpty())
+			{
+				FoldersAddedDelegate.Broadcast(ClassChanges.FoldersModified);
+			}
+			if (!ClassChanges.ClassesModified.IsEmpty())
+			{
+				ClassesAddedDelegate.Broadcast(ClassChanges.ClassesModified);
+			}
 
-				// Then remove the old classes
-				for (const TPair<UClass*, UClass*>& OldAndNew : ReinstancedClasses)
-				{
-					if (OldAndNew.Key && !OldAndNew.Value)
-					{
-						RemoveClass(OldAndNew.Key, GameModules, ClassChanges);
-					}
-					}
+			// Reset the class changes tracking
+			ClassChanges.Reset();
 
-				if (!ClassChanges.ClassesModified.IsEmpty())
-				{
-					ClassesRemovedDelegate.Broadcast(ClassChanges.ClassesModified);
-				}
-				if (!ClassChanges.FoldersModified.IsEmpty())
+			// Then remove the old classes
+			for (const TPair<UClass*, UClass*>& OldAndNew : ReinstancedClasses)
+			{
+				if (OldAndNew.Key && !OldAndNew.Value)
 				{
-					FoldersRemovedDelegate.Broadcast(ClassChanges.FoldersModified);
+					RemoveClass(OldAndNew.Key, GameModules, ClassChanges);
 				}
 			}
+
+			if (!ClassChanges.ClassesModified.IsEmpty())
+			{
+				ClassesRemovedDelegate.Broadcast(ClassChanges.ClassesModified);
+			}
+			if (!ClassChanges.FoldersModified.IsEmpty())
+			{
+				FoldersRemovedDelegate.Broadcast(ClassChanges.FoldersModified);
+			}
 		}
 	}
 }
diff --git a/Engine/Source/Developer/Windows/LiveCoding/Private/LiveCodingModule.cpp b/Engine/Source/Developer/Windows/LiveCoding/Private/LiveCodingModule.cpp
--- a/Engine/Source/Developer/Windows/LiveCoding/Private/LiveCodingModule.cpp
+++ b/Engine/Source/Developer/Windows/LiveCoding/Private/LiveCodingModule.cpp
@@ -139,11 +139,19 @@ public:
 		return CDO;
 	}
 
-	virtual const TMap<UClass*, UClass*>* GetReinstancedClasses() override
+	virtual const TMap<UClass*, UClass*>& GetReinstancedClasses() override
 	{
-		return nullptr;
+		static const TMap<UClass*, UClass*> Empty;
+		return Empty;
 	}
 
+	virtual const TArray<UClass*>& GetNewClasses() override
+	{
+		static const TArray<UClass*> Empty;
+		return Empty;
+	}
+
+
 	bool HasReinstancingOccurred() const
 	{
 		return bHasReinstancingOccurred;
diff --git a/Engine/Source/Editor/UnrealEd/Private/Kismet2/ReloadUtilities.cpp b/Engine/Source/Editor/UnrealEd/Private/Kismet2/ReloadUtilities.cpp
--- a/Engine/Source/Editor/UnrealEd/Private/Kismet2/ReloadUtilities.cpp
+++ b/Engine/Source/Editor/UnrealEd/Private/Kismet2/ReloadUtilities.cpp
@@ -871,6 +871,7 @@ void FReload::Reset()
 	FunctionRemap.Empty();
 	ReconstructedCDOsMap.Empty();
 	ReinstancedClasses.Empty();
+	NewClasses.Empty();
 	ReinstancedEnums.Empty();
 	ReinstancedStructs.Empty();
 	Packages.Empty();
@@ -932,7 +933,6 @@ void FReload::NotifyChange(UClass* New, UClass* Old)
 		bHasReinstancingOccurred = true;
 	}
 
-	// Ignore new classes
 	if (Old != nullptr)
 	{
 		// Don't allow re-instancing of UEngine classes
@@ -948,6 +948,10 @@ void FReload::NotifyChange(UClass* New, UClass* Old)
 			Ar.Logf(ELogVerbosity::Warning, TEXT("Engine class '%s' has changed but will be ignored for reload"), *New->GetName());
 		}
 	}
+	else if (New != nullptr) // should always be the case, but protect against it
+	{
+		NewClasses.Add(New);
+	}
 }
 
 void FReload::NotifyChange(UEnum* New, UEnum* Old)
diff --git a/Engine/Source/Editor/UnrealEd/Public/Kismet2/ReloadUtilities.h b/Engine/Source/Editor/UnrealEd/Public/Kismet2/ReloadUtilities.h
--- a/Engine/Source/Editor/UnrealEd/Public/Kismet2/ReloadUtilities.h
+++ b/Engine/Source/Editor/UnrealEd/Public/Kismet2/ReloadUtilities.h
@@ -44,9 +44,15 @@ public:
 	UNREALED_API virtual void Reinstance() override;
 	UNREALED_API virtual UObject* GetReinstancedCDO(UObject* CDO) override;
 	UNREALED_API virtual const UObject* GetReinstancedCDO(const UObject* CDO) override;
-	UNREALED_API virtual const TMap<UClass*, UClass*>* GetReinstancedClasses() override
+
+	UNREALED_API virtual const TMap<UClass*, UClass*>& GetReinstancedClasses() override
 	{
-		return &ReinstancedClasses;
+		return ReinstancedClasses;
+	}
+
+	UNREALED_API const TArray<UClass*>& GetNewClasses() override
+	{
+		return NewClasses;
 	}
 
 	/**
@@ -153,6 +159,9 @@ private:
 	/** Map from old class to new class.  New class may be null */
 	TMap<UClass*, UClass*> ReinstancedClasses;
 
+	/** Array of new classes */
+	TArray<UClass*> NewClasses;
+
 	/** Map from old struct to new struct.  New struct may be null */
 	TMap<UScriptStruct*, UScriptStruct*> ReinstancedStructs;
 
diff --git a/Engine/Source/Runtime/CoreUObject/Public/UObject/Reload.h b/Engine/Source/Runtime/CoreUObject/Public/UObject/Reload.h
--- a/Engine/Source/Runtime/CoreUObject/Public/UObject/Reload.h
+++ b/Engine/Source/Runtime/CoreUObject/Public/UObject/Reload.h
@@ -62,10 +62,14 @@ public:
 	virtual const UObject* GetReinstancedCDO(const UObject* CDO) = 0;
 
 	/**
-	 * Return an optional map of the reinstanced classes where the old class is the key and the new class is the value.
-	 * Note: the returned value can be null.
+	 * Return a map of the reinstanced classes where the old class is the key and the new class is the value.
 	 */
-	 virtual const TMap<UClass*, UClass*>* GetReinstancedClasses() = 0;
+	 virtual const TMap<UClass*, UClass*>& GetReinstancedClasses() = 0;
+
+	 /**
+	  * Return an array of the new classes
+	  */
+	 virtual const TArray<UClass*>& GetNewClasses() = 0;
 };
 
 /**
