Added support for navigating to script package name folders entered in the content browser navigation bar.

diff --git a/Engine/Plugins/Editor/ContentBrowser/ContentBrowserClassDataSource/Source/ContentBrowserClassDataSource/Private/ContentBrowserClassDataSource.cpp b/Engine/Plugins/Editor/ContentBrowser/ContentBrowserClassDataSource/Source/ContentBrowserClassDataSource/Private/ContentBrowserClassDataSource.cpp
--- a/Engine/Plugins/Editor/ContentBrowser/ContentBrowserClassDataSource/Source/ContentBrowserClassDataSource/Private/ContentBrowserClassDataSource.cpp
+++ b/Engine/Plugins/Editor/ContentBrowser/ContentBrowserClassDataSource/Source/ContentBrowserClassDataSource/Private/ContentBrowserClassDataSource.cpp
@@ -386,21 +386,37 @@ void UContentBrowserClassDataSource::EnumerateItemsAtPath(const FName InPath, co
 
 void UContentBrowserClassDataSource::EnumerateItemsAtUserProvidedPath(const FStringView InPath, const EContentBrowserItemTypeFilter InItemTypeFilter, TFunctionRef<bool(FContentBrowserItemData&&)> InCallback)
 {
-	if (EnumHasAnyFlags(InItemTypeFilter, EContentBrowserItemTypeFilter::IncludeFiles) && InPath.StartsWith(TEXT("/Script/")))
+	if (!InPath.StartsWith(TEXT("/Script/")))
 	{
-		const FTopLevelAssetPath Path(InPath);
-		if (Path.IsValid() && !Path.GetAssetName().IsNone())
+		return;
+	}
+
+	const FTopLevelAssetPath Path(InPath);
+	if (!Path.IsValid())
+	{
+		return;
+	}
+
+	if (EnumHasAnyFlags(InItemTypeFilter, EContentBrowserItemTypeFilter::IncludeFolders) && Path.GetAssetName().IsNone())
+	{
+		const FName FolderPath = NativeClassHierarchy->GetFolderPathForModule(FPackageName::GetShortFName(Path.GetPackageName()), NativeClassHierarchyGetClassPathCache.GameModules);
+		if (TSharedPtr<const FNativeClassHierarchyNode> ClassNode = NativeClassHierarchy->FindNode(FolderPath, ENativeClassHierarchyNodeType::Folder))
 		{
-			if (UClass* Class = FindObject<UClass>(Path))
+			InCallback(CreateClassFolderItem(FolderPath, ClassNode));
+		}
+	}
+
+	if (EnumHasAnyFlags(InItemTypeFilter, EContentBrowserItemTypeFilter::IncludeFiles) && !Path.GetAssetName().IsNone())
+	{
+		if (UClass* Class = FindObject<UClass>(Path))
+		{
+			FString InternalPath;
+			if (NativeClassHierarchy->GetClassPath(Class, InternalPath, NativeClassHierarchyGetClassPathCache.GameModules))
 			{
-				FString InternalPath;
-				if (NativeClassHierarchy->GetClassPath(Class, InternalPath, NativeClassHierarchyGetClassPathCache.GameModules))
+				FName ClassPath = FName(InternalPath);
+				if (TSharedPtr<const FNativeClassHierarchyNode> ClassNode = NativeClassHierarchy->FindNode(ClassPath, ENativeClassHierarchyNodeType::Class))
 				{
-					FName ClassPath = FName(InternalPath);
-					if (TSharedPtr<const FNativeClassHierarchyNode> ClassNode = NativeClassHierarchy->FindNode(ClassPath, ENativeClassHierarchyNodeType::Class))
-					{
-						InCallback(CreateClassFileItem(ClassPath, ClassNode));
-					}
+					InCallback(CreateClassFileItem(ClassPath, ClassNode));
 				}
 			}
 		}
diff --git a/Engine/Plugins/Editor/ContentBrowser/ContentBrowserClassDataSource/Source/ContentBrowserClassDataSource/Private/NativeClassHierarchy.cpp b/Engine/Plugins/Editor/ContentBrowser/ContentBrowserClassDataSource/Source/ContentBrowserClassDataSource/Private/NativeClassHierarchy.cpp
--- a/Engine/Plugins/Editor/ContentBrowser/ContentBrowserClassDataSource/Source/ContentBrowserClassDataSource/Private/NativeClassHierarchy.cpp
+++ b/Engine/Plugins/Editor/ContentBrowser/ContentBrowserClassDataSource/Source/ContentBrowserClassDataSource/Private/NativeClassHierarchy.cpp
@@ -649,6 +649,25 @@ bool FNativeClassHierarchy::GetClassPath(const UClass* InClass, FString& OutClas
 	return true;
 }
 
+FName FNativeClassHierarchy::GetFolderPathForModule(FName InModuleName, TSet<FName>& InGameModuleCache) const
+{
+	if (InGameModuleCache.Num() == 0)
+	{
+		InGameModuleCache = GetGameModules();
+	}
+
+	TOptional<EPluginLoadedFrom> WhereLoadedFrom;
+	const FName RootNodeName = GetClassPathRootForModule(InModuleName, InGameModuleCache, WhereLoadedFrom);
+
+	FNameBuilder FolderPath;
+	FolderPath.AppendChar(TEXT('/'));
+	RootNodeName.AppendString(FolderPath);
+	FolderPath.AppendChar(TEXT('/'));
+	InModuleName.AppendString(FolderPath);
+
+	return FName(FolderPath);
+}
+
 void FNativeClassHierarchy::OnModulesChanged(FName InModuleName, EModuleChangeReason InModuleChangeReason)
 {
 	switch(InModuleChangeReason)
diff --git a/Engine/Plugins/Editor/ContentBrowser/ContentBrowserClassDataSource/Source/ContentBrowserClassDataSource/Private/NativeClassHierarchy.h b/Engine/Plugins/Editor/ContentBrowser/ContentBrowserClassDataSource/Source/ContentBrowserClassDataSource/Private/NativeClassHierarchy.h
--- a/Engine/Plugins/Editor/ContentBrowser/ContentBrowserClassDataSource/Source/ContentBrowserClassDataSource/Private/NativeClassHierarchy.h
+++ b/Engine/Plugins/Editor/ContentBrowser/ContentBrowserClassDataSource/Source/ContentBrowserClassDataSource/Private/NativeClassHierarchy.h
@@ -263,6 +263,16 @@ public:
 	 */
 	bool GetClassPath(const UClass* InClass, FString& OutClassPath, TSet<FName>& InGameModuleCache, const bool bIncludeClassName = true) const;
 
+	/**
+	 * Work out the folder path that should be used for the given module
+	 *
+	 * @param InModuleName - The module we want to find the folder path for
+	 * @param InGameModuleCache - The list of modules that belong to the "/Classes_Game" root
+	 * 
+	 * @return The folder path for the module
+	 */
+	FName GetFolderPathForModule(FName InModuleName, TSet<FName>& InGameModuleCache) const;
+
 	/**
 	 * Test if root node passes given rules
 	 */
